Added a new RenderTrace module to perform precise ray casts using the renderer in order to determine a physical material at a specific point.

Works similar to an AsyncLineTrace except that a list of primitive components must be specified. The primitive components need a material with a UMaterialExpressionPhysicalMaterialOutput node that contains a list of UPhysicalMaterials. A draw is performed with all the primitives, and the nearest hit writes its Physical Material into the render target which is then read back from the GPU. This incurs a few frames of delay in order to not flush the render pipeline.

diff --git a/Engine/Plugins/Runtime/RenderTrace/RenderTrace.uplugin b/Engine/Plugins/Runtime/RenderTrace/RenderTrace.uplugin
new file mode 100644
--- /dev/null
+++ b/Engine/Plugins/Runtime/RenderTrace/RenderTrace.uplugin
@@ -0,0 +1,25 @@
+{
+	"FileVersion": 1,
+	"Version": 1,
+	"VersionName": "1.0",
+	"FriendlyName": "Render Trace",
+	"Description": "The Render Trace plugin provides a way to have pixel perfect sampling of physical materials on meshes.",
+	"Category": "Physics",
+	"CreatedBy": "Epic Games, Inc.",
+	"CreatedByURL": "http://epicgames.com",
+	"DocsURL": "",
+	"MarketplaceURL": "",
+	"SupportURL": "",
+	"EnabledByDefault": false,
+	"CanContainContent": false,
+	"IsBetaVersion": true,
+	"Installed": false,
+	"Modules": [
+		{
+			"Name": "RenderTrace",
+			"Type": "Runtime",
+			"LoadingPhase": "PostConfigInit"
+		}
+	],
+	"IsExperimentalVersion": false
+}
diff --git a/Engine/Plugins/Runtime/RenderTrace/Shaders/Private/PhysicalMaterialSampler.usf b/Engine/Plugins/Runtime/RenderTrace/Shaders/Private/PhysicalMaterialSampler.usf
new file mode 100644
--- /dev/null
+++ b/Engine/Plugins/Runtime/RenderTrace/Shaders/Private/PhysicalMaterialSampler.usf
@@ -0,0 +1,127 @@
+// Copyright Epic Games, Inc. All Rights Reserved.
+
+/**
+ * Material shader used to output the dominant physical materials.
+ */
+
+#include "/Engine/Private/Common.ush"
+#include "/Engine/Generated/Material.ush"
+#include "/Engine/Generated/VertexFactory.ush"
+
+struct FPhysicalMaterialSamplerInterpolantsVSToPS
+{
+	FVertexFactoryInterpolantsVSToPS FactoryInterpolants;
+	float4 Position : SV_POSITION;
+};
+
+#if VERTEXSHADER
+
+/** Simple default vertex shader. */
+void VSMain(
+	FVertexFactoryInput Input,
+	out FPhysicalMaterialSamplerInterpolantsVSToPS Output
+)
+{
+	ResolvedView = ResolveView();
+	FVertexFactoryIntermediates VFIntermediates = GetVertexFactoryIntermediates(Input);
+	float4 WorldPosition = VertexFactoryGetWorldPosition(Input, VFIntermediates);
+
+	half3x3 TangentToLocal = VertexFactoryGetTangentToLocal(Input, VFIntermediates);
+	FMaterialVertexParameters VertexParameters = GetMaterialVertexParameters(Input, VFIntermediates, WorldPosition.xyz, TangentToLocal);
+	Output.FactoryInterpolants = VertexFactoryGetInterpolantsVSToPS(Input, VFIntermediates, VertexParameters);
+
+	WorldPosition.xyz += GetMaterialWorldPositionOffset(VertexParameters);
+	float4 RasterizedWorldPosition = VertexFactoryGetRasterizedWorldPosition(Input, VFIntermediates, WorldPosition);
+	Output.Position = mul(RasterizedWorldPosition, ResolvedView.TranslatedWorldToClip);
+}
+
+#elif PIXELSHADER
+
+float DrawSerialNumber;
+
+float2 UpdateDominantMaterial(in float Index, in float Contribution, in float2 DominantMaterial)
+{
+	if (Contribution > DominantMaterial.y)
+	{
+		DominantMaterial.x = Index + 1; // Store indices as 1-based so that 0 means empty/default.
+		DominantMaterial.y = Contribution;
+	}
+	return DominantMaterial;
+}
+
+/** Pixel shader outputs index of physical material. */
+void PSMain(
+	FPhysicalMaterialSamplerInterpolantsVSToPS Interpolants
+	OPTIONAL_IsFrontFace,
+	out HALF4_TYPE OutColor : SV_Target0
+)
+{
+	ResolvedView = ResolveView();
+
+	float4 SvPosition = Interpolants.Position;
+	FMaterialPixelParameters MaterialParameters = GetMaterialPixelParameters(Interpolants.FactoryInterpolants, SvPosition);
+
+	FPixelMaterialInputs PixelMaterialInputs;
+	CalcMaterialParameters(MaterialParameters, PixelMaterialInputs, SvPosition, bIsFrontFace);
+	GetMaterialCoverageAndClipping(MaterialParameters, PixelMaterialInputs);
+
+	float2 DominantMaterial = float2(0, 0); // float2(MaterialIndex, Contribution)
+
+#if NUM_MATERIAL_OUTPUTS_GETRENDERTRACEPHYSICALMATERIAL > 0
+	DominantMaterial = UpdateDominantMaterial(0, GetRenderTracePhysicalMaterial0(MaterialParameters), DominantMaterial);
+#endif
+#if NUM_MATERIAL_OUTPUTS_GETRENDERTRACEPHYSICALMATERIAL > 1
+	DominantMaterial = UpdateDominantMaterial(1, GetRenderTracePhysicalMaterial1(MaterialParameters), DominantMaterial);
+#endif
+#if NUM_MATERIAL_OUTPUTS_GETRENDERTRACEPHYSICALMATERIAL > 2
+	DominantMaterial = UpdateDominantMaterial(2, GetRenderTracePhysicalMaterial2(MaterialParameters), DominantMaterial);
+#endif
+#if NUM_MATERIAL_OUTPUTS_GETRENDERTRACEPHYSICALMATERIAL > 3
+	DominantMaterial = UpdateDominantMaterial(3, GetRenderTracePhysicalMaterial3(MaterialParameters), DominantMaterial);
+#endif
+#if NUM_MATERIAL_OUTPUTS_GETRENDERTRACEPHYSICALMATERIAL > 4
+	DominantMaterial = UpdateDominantMaterial(4, GetRenderTracePhysicalMaterial4(MaterialParameters), DominantMaterial);
+#endif
+#if NUM_MATERIAL_OUTPUTS_GETRENDERTRACEPHYSICALMATERIAL > 5
+	DominantMaterial = UpdateDominantMaterial(5, GetRenderTracePhysicalMaterial5(MaterialParameters), DominantMaterial);
+#endif
+#if NUM_MATERIAL_OUTPUTS_GETRENDERTRACEPHYSICALMATERIAL > 6
+	DominantMaterial = UpdateDominantMaterial(6, GetRenderTracePhysicalMaterial6(MaterialParameters), DominantMaterial);
+#endif
+#if NUM_MATERIAL_OUTPUTS_GETRENDERTRACEPHYSICALMATERIAL > 7
+	DominantMaterial = UpdateDominantMaterial(7, GetRenderTracePhysicalMaterial7(MaterialParameters), DominantMaterial);
+#endif
+#if NUM_MATERIAL_OUTPUTS_GETRENDERTRACEPHYSICALMATERIAL > 8
+	DominantMaterial = UpdateDominantMaterial(8, GetRenderTracePhysicalMaterial8(MaterialParameters), DominantMaterial);
+#endif
+#if NUM_MATERIAL_OUTPUTS_GETRENDERTRACEPHYSICALMATERIAL > 9
+	DominantMaterial = UpdateDominantMaterial(9, GetRenderTracePhysicalMaterial9(MaterialParameters), DominantMaterial);
+#endif
+#if NUM_MATERIAL_OUTPUTS_GETRENDERTRACEPHYSICALMATERIAL > 10
+	DominantMaterial = UpdateDominantMaterial(10, GetRenderTracePhysicalMaterial10(MaterialParameters), DominantMaterial);
+#endif
+#if NUM_MATERIAL_OUTPUTS_GETRENDERTRACEPHYSICALMATERIAL > 11
+	DominantMaterial = UpdateDominantMaterial(11, GetRenderTracePhysicalMaterial11(MaterialParameters), DominantMaterial);
+#endif
+#if NUM_MATERIAL_OUTPUTS_GETRENDERTRACEPHYSICALMATERIAL > 12
+	DominantMaterial = UpdateDominantMaterial(12, GetRenderTracePhysicalMaterial12(MaterialParameters), DominantMaterial);
+#endif
+#if NUM_MATERIAL_OUTPUTS_GETRENDERTRACEPHYSICALMATERIAL > 13
+	DominantMaterial = UpdateDominantMaterial(13, GetRenderTracePhysicalMaterial13(MaterialParameters), DominantMaterial);
+#endif
+#if NUM_MATERIAL_OUTPUTS_GETRENDERTRACEPHYSICALMATERIAL > 14
+	DominantMaterial = UpdateDominantMaterial(14, GetRenderTracePhysicalMaterial14(MaterialParameters), DominantMaterial);
+#endif
+#if NUM_MATERIAL_OUTPUTS_GETRENDERTRACEPHYSICALMATERIAL > 15
+	DominantMaterial = UpdateDominantMaterial(15, GetRenderTracePhysicalMaterial15(MaterialParameters), DominantMaterial);
+#endif
+#if NUM_MATERIAL_OUTPUTS_GETRENDERTRACEPHYSICALMATERIAL > 16
+	// We pick an arbitrary limit to support here, but the only cost of extending this is more typing.
+#error Too many physical materials
+#endif
+
+	// Write index assuming R16G16_UNORM
+	OutColor = float4(DrawSerialNumber / 65535.0, DominantMaterial.x / 65535.0, 0, 0);
+}
+
+#endif
diff --git a/Engine/Plugins/Runtime/RenderTrace/Source/Private/Materials/MaterialExpressionPhysicalMaterialOutput.cpp b/Engine/Plugins/Runtime/RenderTrace/Source/Private/Materials/MaterialExpressionPhysicalMaterialOutput.cpp
new file mode 100644
--- /dev/null
+++ b/Engine/Plugins/Runtime/RenderTrace/Source/Private/Materials/MaterialExpressionPhysicalMaterialOutput.cpp
@@ -0,0 +1,148 @@
+// Copyright Epic Games, Inc. All Rights Reserved.
+
+#include "Materials/MaterialExpressionPhysicalMaterialOutput.h"
+
+#include "Engine/Engine.h"
+#include "EngineGlobals.h"
+#include "MaterialCompiler.h"
+#include "PhysicalMaterials/PhysicalMaterial.h"
+
+#define LOCTEXT_NAMESPACE "RenderTrace"
+
+UMaterialExpressionPhysicalMaterialOutput::UMaterialExpressionPhysicalMaterialOutput(const FObjectInitializer& ObjectInitializer)
+	: Super(ObjectInitializer)
+{
+	struct FConstructorStatics
+	{
+		FText NAME_RenderTrace;
+		FConstructorStatics()
+			: NAME_RenderTrace(LOCTEXT("RenderTrace", "Render Trace"))
+		{
+		}
+	};
+	static FConstructorStatics ConstructorStatics;
+
+#if WITH_EDITORONLY_DATA
+	MenuCategories.Add(ConstructorStatics.NAME_RenderTrace);
+
+	// No outputs
+	Outputs.Reset();
+#endif
+}
+
+#if WITH_EDITOR
+
+const UPhysicalMaterial* UMaterialExpressionPhysicalMaterialOutput::GetInputMaterialFromMap(int32 Index) const
+{
+	if (Material != nullptr)
+	{
+		return Material->GetPhysicalMaterialFromMap(Index);
+	}
+	return nullptr;
+}
+
+void UMaterialExpressionPhysicalMaterialOutput::FillMaterialNames()
+{
+	if (Material)
+	{
+		TArray<TObjectPtr<class UPhysicalMaterial>> MaterialArray;
+		MaterialArray.Empty(Inputs.Num());
+		for (auto& Input : Inputs)
+		{
+			MaterialArray.Add(Input.PhysicalMaterial);
+		}
+		Material->SetRenderTracePhysicalMaterialOutputs(MaterialArray);
+	}
+}
+
+void UMaterialExpressionPhysicalMaterialOutput::GetCaption(TArray<FString>& OutCaptions) const
+{
+	OutCaptions.Add(FString(TEXT("Physical Material Output")));
+}
+
+const TArray<FExpressionInput*> UMaterialExpressionPhysicalMaterialOutput::GetInputs()
+{
+	TArray<FExpressionInput*> OutInputs;
+	for (auto& Input : Inputs)
+	{
+		OutInputs.Add(&Input.Input);
+	}
+	return OutInputs;
+}
+
+FExpressionInput* UMaterialExpressionPhysicalMaterialOutput::GetInput(int32 InputIndex)
+{
+	return &Inputs[InputIndex].Input;
+}
+
+FName UMaterialExpressionPhysicalMaterialOutput::GetInputName(int32 InputIndex) const
+{
+	UPhysicalMaterial* PhysicalMaterial = Inputs[InputIndex].PhysicalMaterial;
+	if (PhysicalMaterial != nullptr)
+	{
+		return PhysicalMaterial->GetFName();
+	}
+	return NAME_None;
+}
+
+int32 UMaterialExpressionPhysicalMaterialOutput::Compile(class FMaterialCompiler* Compiler, int32 OutputIndex)
+{
+	// In PhysicalMaterialSampler.usf
+	if (OutputIndex > 16)
+	{
+		return CompilerError(Compiler, TEXT("PhysicalMaterialSampler.usf only supports 16 inputs."));
+	}
+
+	if (Inputs.IsValidIndex(OutputIndex))
+	{
+		if (Inputs[OutputIndex].PhysicalMaterial == nullptr)
+		{
+			return CompilerError(Compiler, TEXT("Physical material not set."));
+		}
+
+		if (Inputs[OutputIndex].Input.Expression == nullptr)
+		{
+			return CompilerError(Compiler, TEXT("Input missing."));
+		}
+
+		return Compiler->CustomOutput(this, OutputIndex, Inputs[OutputIndex].Input.Compile(Compiler));
+	}
+
+	return INDEX_NONE;
+}
+
+void UMaterialExpressionPhysicalMaterialOutput::PostLoad()
+{
+	Super::PostLoad();
+	FillMaterialNames();
+}
+
+void UMaterialExpressionPhysicalMaterialOutput::PostEditChangeProperty(FPropertyChangedEvent& PropertyChangedEvent)
+{
+	Super::PostEditChangeProperty(PropertyChangedEvent);
+
+	if (PropertyChangedEvent.MemberProperty)
+	{
+		const FName PropertyName = PropertyChangedEvent.MemberProperty->GetFName();
+		if (PropertyName == GET_MEMBER_NAME_CHECKED(UMaterialExpressionPhysicalMaterialOutput, Inputs))
+		{
+			FillMaterialNames();
+
+			if (GraphNode)
+			{
+				GraphNode->ReconstructNode();
+			}
+		}
+	}
+}
+
+#endif // WITH_EDITOR
+
+int32 UMaterialExpressionPhysicalMaterialOutput::GetNumOutputs() const 
+{ 
+	int32 Count = (int32)Inputs.Num(); 
+	UE_LOG(LogTemp, Warning, TEXT("UMaterialExpressionPhysicalMaterialOutput::GetNumOutputs() is returning %d"), Count);
+	return Inputs.Num(); 
+}
+
+#undef LOCTEXT_NAMESPACE
diff --git a/Engine/Plugins/Runtime/RenderTrace/Source/Private/Materials/MaterialExpressionPhysicalMaterialOutput.h b/Engine/Plugins/Runtime/RenderTrace/Source/Private/Materials/MaterialExpressionPhysicalMaterialOutput.h
new file mode 100644
--- /dev/null
+++ b/Engine/Plugins/Runtime/RenderTrace/Source/Private/Materials/MaterialExpressionPhysicalMaterialOutput.h
@@ -0,0 +1,62 @@
+// Copyright Epic Games, Inc. All Rights Reserved.
+
+#pragma once
+
+#include "CoreMinimal.h"
+#include "Materials/MaterialExpressionCustomOutput.h"
+
+#include "MaterialExpressionPhysicalMaterialOutput.generated.h"
+
+/** Structure linking a material expression input with a physical material. For use by UMaterialExpressionPhysicalMaterialOutput. */
+USTRUCT()
+struct FPhysicalMaterialTraceInput
+{
+	GENERATED_USTRUCT_BODY()
+
+	UPROPERTY(EditAnywhere, Category = PhysicalMaterial)
+	TObjectPtr<UPhysicalMaterial> PhysicalMaterial;
+
+	UPROPERTY(meta = (RequiredInput = "true"))
+	FExpressionInput Input;
+
+	FPhysicalMaterialTraceInput()
+		: PhysicalMaterial(nullptr)
+	{}
+};
+
+/** 
+ * Custom output node to write out physical material weights.
+ */
+UCLASS(collapsecategories, hidecategories=Object)
+class RENDERTRACE_API UMaterialExpressionPhysicalMaterialOutput : public UMaterialExpressionCustomOutput
+{
+	GENERATED_UCLASS_BODY()
+
+	/** Array of physical material inputs. */
+	UPROPERTY(EditAnywhere, Category = PhysicalMaterial)
+	TArray<FPhysicalMaterialTraceInput> Inputs;
+
+#if WITH_EDITOR
+	const UPhysicalMaterial* GetInputMaterialFromMap(int32 Index) const;
+	void FillMaterialNames();
+
+	//~ Begin UObject Interface
+	virtual void PostEditChangeProperty(FPropertyChangedEvent& PropertyChangedEvent) override;
+	virtual void PostLoad() override;
+	//~ End UObject Interface
+
+	//~ Begin UMaterialExpression Interface
+	virtual void GetCaption(TArray<FString>& OutCaptions) const override;
+	virtual const TArray<FExpressionInput*> GetInputs() override;
+	virtual FExpressionInput* GetInput(int32 InputIndex) override;
+	virtual FName GetInputName(int32 InputIndex) const override;
+	virtual uint32 GetInputType(int32 InputIndex) override { return MCT_Float; }
+	virtual int32 Compile(class FMaterialCompiler* Compiler, int32 OutputIndex) override;
+	//~ End UMaterialExpression Interface
+#endif
+
+	//~ Begin UMaterialExpressionCustomOutput Interface
+	virtual int32 GetNumOutputs() const override;
+	virtual FString GetFunctionName() const override { return TEXT("GetRenderTracePhysicalMaterial"); }
+	//~ End UMaterialExpressionCustomOutput Interface
+};
diff --git a/Engine/Plugins/Runtime/RenderTrace/Source/Private/RenderTrace.cpp b/Engine/Plugins/Runtime/RenderTrace/Source/Private/RenderTrace.cpp
new file mode 100644
--- /dev/null
+++ b/Engine/Plugins/Runtime/RenderTrace/Source/Private/RenderTrace.cpp
@@ -0,0 +1,684 @@
+// Copyright Epic Games, Inc. All Rights Reserved.
+
+#include "RenderTrace.h"
+
+#include "EngineModule.h"
+#include "Components/PrimitiveComponent.h"
+#include "PrimitiveSceneProxy.h"
+#include "Materials/MaterialExpressionPhysicalMaterialOutput.h"
+#include "Materials/MaterialInstanceConstant.h"
+#include "MeshMaterialShader.h"
+#include "MeshPassProcessor.h"
+#include "MeshPassProcessor.inl"
+#include "RenderCommandFence.h"
+#include "RenderTargetPool.h"
+#include "RHIResources.h"
+#include "SceneRenderTargetParameters.h"
+#include "SimpleMeshDrawCommandPass.h"
+#include "Logging/LogMacros.h"
+
+DEFINE_LOG_CATEGORY(LogRenderTrace);
+
+DECLARE_GPU_STAT_NAMED(RenderTraceDraw, TEXT("RenderTrace Draw"));
+DECLARE_GPU_STAT_NAMED(RenderTraceReadback, TEXT("RenderTrace Readback"));
+
+DECLARE_STATS_GROUP(TEXT("RenderTrace"), STATGROUP_RenderTrace, STATCAT_Advanced);
+DECLARE_DWORD_COUNTER_STAT(TEXT("RenderTrace Requested"), STAT_RenderTrace_Requests, STATGROUP_RenderTrace);
+DECLARE_DWORD_COUNTER_STAT(TEXT("RenderTrace Skipped"), STAT_RenderTrace_Skipped, STATGROUP_RenderTrace);
+DECLARE_DWORD_COUNTER_STAT(TEXT("RenderTrace Completed"), STAT_RenderTrace_Completed, STATGROUP_RenderTrace);
+DECLARE_CYCLE_STAT(TEXT("RenderTrace_Submission"), STAT_RenderTrace_Submit, STATGROUP_RenderTrace);
+DECLARE_CYCLE_STAT(TEXT("RenderTrace_Tick"), STAT_RenderTrace_Tick, STATGROUP_RenderTrace);
+
+#define RENDER_TRACE_LOG_TIMING STATS
+
+bool bEnableRenderTracing = false;
+FAutoConsoleVariableRef CVarRenderTraceEnable(
+	TEXT("RenderTrace.Enabled"),
+	bEnableRenderTracing,
+	TEXT("Enable Render Tracing."),
+	ECVF_Default);
+
+
+namespace
+{
+	/** Completion state for a physical material render task. */
+	enum class ECompletionState : uint64
+	{
+		NotStarted = 0,	// Draw not submitted
+		Pending = 1,	// Draw submitted, waiting for GPU
+		Complete = 2	// Result copied back from GPU
+	};
+}
+
+struct FPhysicalMaterialMapping
+{
+	const UPrimitiveComponent* PrimitiveComponent;
+	TArray<const UPhysicalMaterial*> PhysicalMaterials;
+};
+
+struct FRenderTraceReadbackData
+{
+	uint16 MaterialMapIndex;
+	uint16 PhysicalMaterialIndex;
+};
+static_assert(sizeof(FRenderTraceReadbackData) == 4, "FRenderTraceReadbackData must match PF_G16R16");
+
+/** Data for a physical material render task. */
+struct FRenderTraceTask
+{
+	// Create on game thread
+	TArray<FPhysicalMaterialMapping> MaterialMap;
+	uint32 RequestID = 0;
+	FIntPoint TargetSize = FIntPoint(ForceInitToZero);
+	FVector ViewOrigin = FVector(ForceInitToZero);
+	FMatrix ViewRotationMatrix = FMatrix(ForceInitToZero);
+	FMatrix ProjectionMatrix = FMatrix(ForceInitToZero);
+
+	// Create on render thread
+	TUniquePtr<FRHIGPUTextureReadback> Readback;
+
+	std::atomic<ECompletionState> CompletionState{ ECompletionState::NotStarted };
+
+	// Result reporting
+	FRenderTraceDelegate ResultDelegate;
+	int64 UserData = 0;
+	const UPhysicalMaterial* ResultMaterial = nullptr;
+
+#if RENDER_TRACE_LOG_TIMING
+	double StartSeconds = 0;
+	uint32 StartFrame = 0;
+#endif
+};
+
+/** 
+ * Get the physical materials attached to the UMaterialExpressionPhysicalMaterialOutput.
+ * Returns false if there are no physical materials.
+ */
+static bool GetPhysicalMaterials(UPrimitiveComponent const* InPrimitiveComponent, TArray<FPhysicalMaterialMapping>& OutMapping)
+{
+	bool bReturnValue = false;
+
+	int32 NumMaterials = InPrimitiveComponent->GetNumMaterials();
+	for (int32 MaterialIndex = 0; MaterialIndex < NumMaterials; ++MaterialIndex)
+	{
+		UMaterialInterface* PrimitiveMaterial = InPrimitiveComponent->GetMaterial(MaterialIndex);
+		if (PrimitiveMaterial == nullptr)
+		{
+			continue;
+		}
+
+		UMaterial* Material = PrimitiveMaterial->GetMaterial();
+		if (Material == nullptr)
+		{
+			continue;
+		}
+
+		TArrayView<const TObjectPtr<UPhysicalMaterial>> PhysicalMaterials = Material->GetRenderTracePhysicalMaterialOutputs();
+		if (!PhysicalMaterials.IsEmpty())
+		{
+			FPhysicalMaterialMapping NewMapping;
+			NewMapping.PrimitiveComponent = InPrimitiveComponent;
+			NewMapping.PhysicalMaterials = PhysicalMaterials;
+			OutMapping.Add(MoveTemp(NewMapping));
+
+			bReturnValue = true;
+		}
+	}
+
+	return bReturnValue;
+}
+
+
+/** Material shader for rendering physical material IDs. */
+class FPhysicalMaterialSamplerShader : public FMeshMaterialShader
+{
+public:
+	FPhysicalMaterialSamplerShader()
+	{}
+
+	FPhysicalMaterialSamplerShader(const FMeshMaterialShaderType::CompiledShaderInitializerType& Initializer)
+		: FMeshMaterialShader(Initializer)
+	{
+		PassUniformBuffer.Bind(Initializer.ParameterMap, FSceneTextureUniformParameters::StaticStructMetadata.GetShaderVariableName());
+	}
+
+	static bool ShouldCompilePermutation(const FMeshMaterialShaderPermutationParameters& Parameters)
+	{
+		return Parameters.MaterialParameters.bHasRenderTracePhysicalMaterialOutput && IsFeatureLevelSupported(Parameters.Platform, ERHIFeatureLevel::SM5);
+	}
+};
+
+class FPhysicalMaterialSamplerShaderVS : public FPhysicalMaterialSamplerShader
+{
+	DECLARE_SHADER_TYPE(FPhysicalMaterialSamplerShaderVS, MeshMaterial);
+
+public:
+	FPhysicalMaterialSamplerShaderVS()
+	{}
+
+	FPhysicalMaterialSamplerShaderVS(const FMeshMaterialShaderType::CompiledShaderInitializerType& Initializer)
+		: FPhysicalMaterialSamplerShader(Initializer)
+	{
+	}
+};
+
+//IMPLEMENT_MATERIAL_SHADER_TYPE(, FPhysicalMaterialSamplerShaderVS, TEXT("/Engine/Private/PhysicalMaterialSampler.usf"), TEXT("VSMain"), SF_Vertex);
+IMPLEMENT_MATERIAL_SHADER_TYPE(, FPhysicalMaterialSamplerShaderVS, TEXT("/Plugin/Runtime/RenderTrace/Private/PhysicalMaterialSampler.usf"), TEXT("VSMain"), SF_Vertex);
+
+struct FPhysicalMaterialSamplerShaderElementData : public FMeshMaterialShaderElementData
+{
+	FPhysicalMaterialSamplerShaderElementData() : MaterialMapIndex(uint32(-1)) {}
+
+	uint32 MaterialMapIndex;
+};
+
+class FPhysicalMaterialSamplerShaderPS : public FPhysicalMaterialSamplerShader
+{
+	DECLARE_SHADER_TYPE(FPhysicalMaterialSamplerShaderPS, MeshMaterial);
+
+public:
+	FPhysicalMaterialSamplerShaderPS()
+	{}
+
+	FPhysicalMaterialSamplerShaderPS(const ShaderMetaType::CompiledShaderInitializerType& Initializer)
+		: FPhysicalMaterialSamplerShader(Initializer)
+	{
+		SerialNumberParameter.Bind(Initializer.ParameterMap, TEXT("DrawSerialNumber"), SPF_Mandatory);
+	}
+
+	void GetShaderBindings(
+		const FScene* Scene,
+		ERHIFeatureLevel::Type FeatureLevel,
+		const FPrimitiveSceneProxy* PrimitiveSceneProxy,
+		const FMaterialRenderProxy& MaterialRenderProxy,
+		const FMaterial& Material,
+		const FMeshPassProcessorRenderState& DrawRenderState,
+		const FPhysicalMaterialSamplerShaderElementData& ShaderElementData,
+		FMeshDrawSingleShaderBindings& ShaderBindings) const
+	{
+		check(ShaderElementData.MaterialMapIndex != -1);
+		FMeshMaterialShader::GetShaderBindings(Scene, FeatureLevel, PrimitiveSceneProxy, MaterialRenderProxy, Material, DrawRenderState, ShaderElementData, ShaderBindings);
+
+		ShaderBindings.Add(SerialNumberParameter, ShaderElementData.MaterialMapIndex);
+	}
+
+	void GetElementShaderBindings(
+		const FShaderMapPointerTable& PointerTable,
+		const FScene* Scene,
+		const FSceneView* ViewIfDynamicMeshCommand,
+		const FVertexFactory* VertexFactory,
+		const EVertexInputStreamType InputStreamType,
+		const FStaticFeatureLevel FeatureLevel,
+		const FPrimitiveSceneProxy* PrimitiveSceneProxy,
+		const FMeshBatch& MeshBatch,
+		const FMeshBatchElement& BatchElement,
+		const FPhysicalMaterialSamplerShaderElementData& ShaderElementData,
+		FMeshDrawSingleShaderBindings& ShaderBindings,
+		FVertexInputStreamArray& VertexStreams) const
+	{
+		FMeshMaterialShader::GetElementShaderBindings(PointerTable, Scene, ViewIfDynamicMeshCommand, VertexFactory, InputStreamType, FeatureLevel, PrimitiveSceneProxy, MeshBatch, BatchElement, ShaderElementData, ShaderBindings, VertexStreams);
+	}
+
+	LAYOUT_FIELD(FShaderParameter, SerialNumberParameter);
+};
+
+//IMPLEMENT_MATERIAL_SHADER_TYPE(, FPhysicalMaterialSamplerShaderPS, TEXT("/Engine/Private/PhysicalMaterialSampler.usf"), TEXT("PSMain"), SF_Pixel);
+IMPLEMENT_MATERIAL_SHADER_TYPE(, FPhysicalMaterialSamplerShaderPS, TEXT("/Plugin/Runtime/RenderTrace/Private/PhysicalMaterialSampler.usf"), TEXT("PSMain"), SF_Pixel);
+
+
+class FRenderTraceMeshProcessor : public FMeshPassProcessor
+{
+public:
+	FRenderTraceMeshProcessor(
+		const FScene* Scene, 
+		const FSceneView* InViewIfDynamicMeshCommand, 
+		FMeshPassDrawListContext* InDrawListContext);
+
+	virtual void AddMeshBatch(
+		const FMeshBatch& RESTRICT MeshBatch,
+		uint64 BatchElementMask,
+		const FPrimitiveSceneProxy* RESTRICT PrimitiveSceneProxy,
+		int32 StaticMeshId = -1) override final;
+
+	FMeshPassProcessorRenderState PassDrawRenderState;
+	uint32 MaterialMapIndex = 0;
+};
+
+FRenderTraceMeshProcessor::FRenderTraceMeshProcessor(
+	const FScene* Scene, 
+	const FSceneView* InViewIfDynamicMeshCommand, 
+	FMeshPassDrawListContext* InDrawListContext)
+	: FMeshPassProcessor(Scene, InViewIfDynamicMeshCommand->GetFeatureLevel(), InViewIfDynamicMeshCommand, InDrawListContext)
+{
+	PassDrawRenderState.SetBlendState(TStaticBlendState<>::GetRHI());
+	PassDrawRenderState.SetDepthStencilState(TStaticDepthStencilState<>::GetRHI());
+}
+
+void FRenderTraceMeshProcessor::AddMeshBatch(
+	const FMeshBatch& RESTRICT MeshBatch,
+	uint64 BatchElementMask,
+	const FPrimitiveSceneProxy* RESTRICT PrimitiveSceneProxy,
+	int32 StaticMeshId)
+{
+	const FMaterialRenderProxy* MaterialRenderProxy = MeshBatch.MaterialRenderProxy;
+	while (MaterialRenderProxy)
+	{
+		const FMaterial* Material = MaterialRenderProxy->GetMaterialNoFallback(FeatureLevel);
+		if (Material)
+		{
+			const FVertexFactory* VertexFactory = MeshBatch.VertexFactory;
+
+			TMeshProcessorShaders<
+				FPhysicalMaterialSamplerShaderVS,
+				FPhysicalMaterialSamplerShaderPS> PassShaders;
+
+			FMaterialShaderTypes ShaderTypes;
+			ShaderTypes.AddShaderType<FPhysicalMaterialSamplerShaderVS>();
+			ShaderTypes.AddShaderType<FPhysicalMaterialSamplerShaderPS>();
+
+			FMaterialShaders Shaders;
+			if (!Material->TryGetShaders(ShaderTypes, VertexFactory->GetType(), Shaders))
+			{
+				MaterialRenderProxy = MaterialRenderProxy->GetFallback(FeatureLevel);
+				continue;
+			}
+			Shaders.TryGetVertexShader(PassShaders.VertexShader);
+			Shaders.TryGetPixelShader(PassShaders.PixelShader);
+
+			const FMeshDrawingPolicyOverrideSettings OverrideSettings = ComputeMeshOverrideSettings(MeshBatch);
+			const ERasterizerFillMode MeshFillMode = ComputeMeshFillMode(MeshBatch, *Material, OverrideSettings);
+			const ERasterizerCullMode MeshCullMode = CM_None;
+
+			FPhysicalMaterialSamplerShaderElementData ShaderElementData;
+			ShaderElementData.InitializeMeshMaterialData(ViewIfDynamicMeshCommand, PrimitiveSceneProxy, MeshBatch, -1, true);
+			ShaderElementData.MaterialMapIndex = MaterialMapIndex++;
+
+			const FMeshDrawCommandSortKey SortKey = CalculateMeshStaticSortKey(PassShaders.VertexShader, PassShaders.PixelShader);
+
+			BuildMeshDrawCommands(
+				MeshBatch,
+				BatchElementMask,
+				PrimitiveSceneProxy,
+				*MaterialRenderProxy,
+				*Material,
+				PassDrawRenderState,
+				PassShaders,
+				MeshFillMode,
+				MeshCullMode,
+				SortKey,
+				EMeshPassFeatures::Default,
+				ShaderElementData);
+
+			break;
+		}
+
+		MaterialRenderProxy = MaterialRenderProxy->GetFallback(FeatureLevel);
+	}
+}
+
+namespace
+{
+	/** A description of a mesh to render. */
+	struct FMeshInfo
+	{
+		FPrimitiveSceneProxy const* Proxy;
+		FMeshBatch MeshBatch;
+		uint64 MeshBatchElementMask;
+		uint32 MaterialMapIndex;
+	};
+	typedef TArray<FMeshInfo, TInlineAllocator<1>> FMeshInfoArray;
+
+	/**
+	 * Collect the meshes to render for a component.
+	 * WARNING: This gets the SceneProxy pointer from the UComponent on the render thread. This doesn't feel safe but it's what the grass renderer does...
+	 */
+	void AddMeshInfos_RenderThread(const FPrimitiveSceneProxy* InSceneProxy, FMeshInfoArray& OutMeshInfos, int64 MaterialMappingIndex)
+	{
+		int32 LODIndex = 0;
+		TArray<FMeshBatch> OutMeshElements;
+		InSceneProxy->GetMeshDescription(LODIndex, OutMeshElements);
+
+		if (OutMeshElements.Num() > 0)
+		{ 
+			FMeshInfo& NewMeshInfo = OutMeshInfos.AddDefaulted_GetRef();
+			NewMeshInfo.Proxy = InSceneProxy;
+			NewMeshInfo.MeshBatch = OutMeshElements[0];
+			NewMeshInfo.MeshBatchElementMask = 1 << 0; // LOD 0 only
+			NewMeshInfo.MaterialMapIndex = MaterialMappingIndex;
+			if (OutMeshElements.Num() > 1)
+			{
+				UE_LOG(LogRenderTrace, Warning, TEXT("Found a mesh with %d elements, only handling the first one"), (int32)OutMeshElements.Num());
+			}
+		}
+	}
+	
+	BEGIN_SHADER_PARAMETER_STRUCT(FPhysicalMaterialSamplerPassParameters, )
+		SHADER_PARAMETER_STRUCT_REF(FViewUniformShaderParameters, View)
+		SHADER_PARAMETER_STRUCT_INCLUDE(FInstanceCullingDrawParams, InstanceCullingDrawParams)
+		RENDER_TARGET_BINDING_SLOTS()
+	END_SHADER_PARAMETER_STRUCT()
+
+	/** Render the physical material IDs and copy to the read back texture. */
+	void Render_RenderThread(
+		FRHICommandListImmediate& RHICmdList,
+		FSceneInterface* SceneInterface,
+		FMeshInfoArray const& MeshInfos,
+		FIntPoint TargetSize,
+		FVector ViewOrigin,
+		FMatrix ViewRotationMatrix,
+		FMatrix ProjectionMatrix,
+		FRHIGPUTextureReadback* Readback)
+	{
+		FMemMark Mark(FMemStack::Get());
+		FRDGBuilder GraphBuilder(RHICmdList);
+		
+		// Create the view
+		FSceneViewFamily::ConstructionValues ViewFamilyInit(nullptr, SceneInterface, FEngineShowFlags(ESFIM_Game));
+		ViewFamilyInit.SetWorldTimes(0.0f, 0.0f, 0.0f);
+		FSceneViewFamilyContext ViewFamily(ViewFamilyInit);
+
+		FScenePrimitiveRenderingContextScopeHelper ScenePrimitiveRenderingContextScopeHelper(GetRendererModule().BeginScenePrimitiveRendering(GraphBuilder, &ViewFamily));
+
+		FSceneViewInitOptions ViewInitOptions;
+		ViewInitOptions.SetViewRectangle(FIntRect(0, 0, TargetSize.X, TargetSize.Y));
+		ViewInitOptions.ViewOrigin = ViewOrigin;
+		ViewInitOptions.ViewRotationMatrix = ViewRotationMatrix;
+		ViewInitOptions.ProjectionMatrix = ProjectionMatrix;
+		ViewInitOptions.ViewFamily = &ViewFamily;
+
+		// PF_G16R16 Must match FRenderTraceReadbackData
+		FRDGTextureRef OutputTexture = GraphBuilder.CreateTexture(
+			FRDGTextureDesc::Create2D(TargetSize, PF_G16R16, FClearValueBinding::Black, TexCreate_RenderTargetable | TexCreate_ShaderResource),
+			TEXT("RenderTraceTarget"));
+
+		FRDGTextureRef DepthTexture = GraphBuilder.CreateTexture(
+			FRDGTextureDesc::Create2D(TargetSize, PF_DepthStencil, FClearValueBinding::DepthFar, TexCreate_DepthStencilTargetable),
+			TEXT("RenderTraceDepthTarget"));
+
+		GetRendererModule().CreateAndInitSingleView(RHICmdList, &ViewFamily, &ViewInitOptions);
+		const FSceneView* View = ViewFamily.Views[0];
+		{
+			RDG_GPU_MASK_SCOPE(GraphBuilder, View->GPUMask);
+			RDG_EVENT_SCOPE(GraphBuilder, "RenderTraceRender");
+			RDG_GPU_STAT_SCOPE(GraphBuilder, RenderTraceDraw);
+
+			auto* PassParameters = GraphBuilder.AllocParameters<FPhysicalMaterialSamplerPassParameters>();
+			PassParameters->View = View->ViewUniformBuffer;
+			PassParameters->RenderTargets[0] = FRenderTargetBinding(OutputTexture, ERenderTargetLoadAction::EClear);
+			PassParameters->RenderTargets.DepthStencil = FDepthStencilBinding(DepthTexture, ERenderTargetLoadAction::EClear, ERenderTargetLoadAction::ENoAction, FExclusiveDepthStencil::DepthWrite_StencilNop);
+
+			AddSimpleMeshPass(GraphBuilder, PassParameters, SceneInterface->GetRenderScene(), *View, nullptr, RDG_EVENT_NAME("RenderTrace"), View->UnscaledViewRect,
+				[View, &MeshInfos](FDynamicPassMeshDrawListContext* DynamicMeshPassContext)
+				{
+					FRenderTraceMeshProcessor PassMeshProcessor(nullptr, View, DynamicMeshPassContext);
+
+					for (auto& MeshInfo : MeshInfos)
+					{
+						const FMeshBatch& Mesh = MeshInfo.MeshBatch;
+						if (Mesh.MaterialRenderProxy != nullptr)
+						{
+							Mesh.MaterialRenderProxy->UpdateUniformExpressionCacheIfNeeded(View->GetFeatureLevel());
+							PassMeshProcessor.MaterialMapIndex = MeshInfo.MaterialMapIndex;
+							PassMeshProcessor.AddMeshBatch(Mesh, MeshInfo.MeshBatchElementMask, MeshInfo.Proxy);
+						}
+					}
+				});
+		}
+
+		{
+			RDG_GPU_MASK_SCOPE(GraphBuilder, View->GPUMask);
+			RDG_EVENT_SCOPE(GraphBuilder, "RenderTraceCopy");
+			RDG_GPU_STAT_SCOPE(GraphBuilder, RenderTraceReadback);
+			AddEnqueueCopyPass(GraphBuilder, Readback, OutputTexture);
+		}
+
+		GraphBuilder.Execute();
+	}
+
+	/** Fetch the physical material IDs from a read back texture. */
+	void FetchResults_RenderThread(
+		FRHICommandListImmediate& RHICmdList,
+		FIntPoint TargetSize,
+		FRHIGPUTextureReadback* Readback,
+		FRenderTraceReadbackData& OutSampleResult)
+	{
+		void* Data = nullptr;
+		int32 Pitch;
+		Readback->LockTexture(RHICmdList, Data, Pitch);
+		check(Data && TargetSize.X <= Pitch);
+		FMemory::Memcpy(&OutSampleResult, Data, sizeof(OutSampleResult));
+		Readback->Unlock();
+	}
+
+	/** Update the physical material render task on the render thread. */
+	void UpdateTask_RenderThread(FRHICommandListImmediate& RHICmdList, FRenderTraceTask& Task, bool bFlush)
+	{
+		// WARNING: We access the UComponent to get in SceneProxy for AddMeshInfos_RenderThread(). 
+		// This isn't good style but probably works since the UComponent owns the update task and is guaranteed to be valid.
+
+
+		ECompletionState CompletionState = Task.CompletionState.load();
+		if (CompletionState == ECompletionState::Pending)
+		{
+			bool bIsReady = Task.Readback->IsReady();
+			if (bFlush || bIsReady)
+			{
+				if (!bIsReady)
+				{
+					RHICmdList.ImmediateFlush(EImmediateFlushType::FlushRHIThread);
+				}
+
+				FRenderTraceReadbackData Result = {};
+				FetchResults_RenderThread(RHICmdList, Task.TargetSize, Task.Readback.Get(), Result);
+
+				if (Task.MaterialMap.IsValidIndex(Result.MaterialMapIndex))
+				{
+					int32 MaterialIndex = (int32)Result.PhysicalMaterialIndex - 1;
+					const TArray<const UPhysicalMaterial*>& PhysicalMaterials = Task.MaterialMap[Result.MaterialMapIndex].PhysicalMaterials;
+					if (PhysicalMaterials.IsValidIndex(MaterialIndex))
+					{
+						Task.ResultMaterial = PhysicalMaterials[MaterialIndex];
+					}
+					else
+					{
+						UE_LOG(LogRenderTrace, Warning, TEXT("Completed a task but had an invalid material index (%d:%d)"), Result.MaterialMapIndex, Result.PhysicalMaterialIndex);
+					}
+				}
+				else
+				{
+					UE_LOG(LogRenderTrace, Warning, TEXT("Completed a task but had an invalid material map index (%d:%d)"), Result.MaterialMapIndex, Result.PhysicalMaterialIndex);
+				}
+
+				Task.CompletionState.store(ECompletionState::Complete);
+			}
+		}
+		else if (CompletionState == ECompletionState::NotStarted)
+		{
+			if (!Task.Readback.IsValid())
+			{
+				Task.Readback = MakeUnique<FRHIGPUTextureReadback>(TEXT("PhysicalMaterialSamplerReadback"));
+			}
+
+			// Render the pending item.
+			FMeshInfoArray MeshInfos;
+			FSceneInterface* SceneInterface = nullptr;
+
+			int32 MappingCount = Task.MaterialMap.Num();
+			for (int32 MappingIndex = 0; MappingIndex < MappingCount; ++MappingIndex)
+			{
+				const FPhysicalMaterialMapping& Mapping = Task.MaterialMap[MappingIndex];
+				const UPrimitiveComponent* PrimitiveComponent = Mapping.PrimitiveComponent;
+				const FPrimitiveSceneProxy* SceneProxy = PrimitiveComponent->SceneProxy;
+				if (SceneProxy)
+				{
+					AddMeshInfos_RenderThread(SceneProxy, MeshInfos, MappingIndex);
+					check(SceneInterface == nullptr || SceneInterface == &SceneProxy->GetScene());
+					SceneInterface = &SceneProxy->GetScene();
+				}
+			}
+
+			check(Task.Readback.IsValid());
+			Render_RenderThread(
+				RHICmdList,
+				SceneInterface,
+				MeshInfos,
+				Task.TargetSize,
+				Task.ViewOrigin,
+				Task.ViewRotationMatrix,
+				Task.ProjectionMatrix,
+				Task.Readback.Get());
+
+			Task.CompletionState.store(ECompletionState::Pending);
+		}
+
+	}
+
+	void UpdateTasks_RenderThread(FRHICommandListImmediate& RHICmdList, const TArray<TSharedPtr<FRenderTraceTask>>& TaskList, bool bFlush)
+	{
+		for (const TSharedPtr<FRenderTraceTask>& Task : TaskList)
+		{
+			check(Task->MaterialMap.Num() > 0);
+			UpdateTask_RenderThread(RHICmdList, *Task, bFlush);
+		}
+	}
+}
+
+
+uint32 FRenderTraceQueue::AsyncRenderTraceComponents(TArrayView<const class UPrimitiveComponent*> PrimitiveComponents, FVector RayOrigin, FVector RayDirection, FRenderTraceDelegate OnComplete, int64 UserData)
+{
+	if (!ensure(bEnableRenderTracing))
+	{
+		return 0;
+	}
+
+	SCOPE_CYCLE_COUNTER(STAT_RenderTrace_Submit);
+	SCOPED_NAMED_EVENT_TEXT("FRenderTraceQueue::AsyncRenderTraceComponents", FColor::Orange);
+
+	INC_DWORD_STAT(STAT_RenderTrace_Requests);
+	if (!RayDirection.Normalize())
+	{
+		UE_LOG(LogRenderTrace, Warning, TEXT("AsyncSampleComponents given invalid RayDirection"));
+		INC_DWORD_STAT(STAT_RenderTrace_Skipped);
+		return 0;
+	}
+
+	TArray<FPhysicalMaterialMapping> MaterialMapping;
+	for (const UPrimitiveComponent* InPrimitiveComponent : PrimitiveComponents)
+	{
+		if (ensure(InPrimitiveComponent))
+		{
+			GetPhysicalMaterials(InPrimitiveComponent, MaterialMapping);
+		}
+	}
+
+	if (MaterialMapping.IsEmpty())
+	{
+		INC_DWORD_STAT(STAT_RenderTrace_Skipped);
+		return 0;
+	}
+
+	uint32 TaskRequestID = ++LastRequestID;
+
+	TSharedPtr<FRenderTraceTask> NewTask = MakeShared<FRenderTraceTask>();
+	NewTask->MaterialMap = MoveTemp(MaterialMapping);
+	NewTask->CompletionState.store(ECompletionState::NotStarted);
+	NewTask->RequestID = TaskRequestID;
+	NewTask->ResultDelegate = MoveTemp(OnComplete);
+	NewTask->UserData = UserData;
+
+#if STATS
+	NewTask->StartSeconds = FPlatformTime::Seconds();
+	NewTask->StartFrame = GFrameCounter;
+#endif
+
+	const FIntPoint TargetSize(1, 1);
+	const FVector TargetCenter = RayOrigin + RayDirection.Normalize();
+	const FVector TargetExtent = FVector(TargetSize, 0.0f);
+
+	FVector UpVector = FVector::UpVector;
+	if (FMath::IsNearlyEqual(FMath::Abs(FVector::DotProduct(UpVector, RayDirection.GetSafeNormal(SMALL_NUMBER, UpVector))), 1.0f))
+	{
+		/* If we're too close to being colinear then pick another vector to do the cross products with. */
+		UpVector = FVector::RightVector;
+	}
+	FMatrix ViewMatrix = FLookFromMatrix(FVector::ZeroVector, RayDirection, UpVector);
+
+	const float ZOffset = WORLD_MAX;
+	const FMatrix ProjectionMatrix = FReversedZOrthoMatrix(TargetExtent.X, TargetExtent.Y, 0.5f / ZOffset, ZOffset);
+
+	NewTask->TargetSize = TargetSize;
+	NewTask->ViewOrigin = TargetCenter;
+	NewTask->ViewRotationMatrix = ViewMatrix;
+	NewTask->ProjectionMatrix = ProjectionMatrix;
+
+	UE_LOG(LogRenderTrace, Verbose, TEXT("AsyncSampleComponents: Queueing new task %u(%p) with %d primitives"), TaskRequestID, NewTask.Get(), PrimitiveComponents.Num());
+	RequestsInFlight.Add(MoveTemp(NewTask));
+	return TaskRequestID;
+}
+
+void FRenderTraceQueue::CancelAsyncSample(uint32 RequestID)
+{
+	RequestsInFlight.RemoveAll(
+		[RequestID](const TSharedPtr<FRenderTraceTask>& Task)
+		{
+			return Task->RequestID == RequestID;
+		});
+}
+
+void FRenderTraceQueue::Tick(float DeltaTime)
+{
+	SCOPE_CYCLE_COUNTER(STAT_RenderTrace_Tick);
+	SCOPED_NAMED_EVENT_TEXT("FRenderTraceQueue::Tick", FColor::Orange);
+
+	check(IsInGameThread());
+	TArray<TSharedPtr<FRenderTraceTask>> RequestsToUpdate;
+	for (TSharedPtr<FRenderTraceTask>& Task : RequestsInFlight)
+	{
+		if (Task->CompletionState.load() != ECompletionState::Complete)
+		{
+			RequestsToUpdate.Add(Task);
+		}
+	}
+
+	if (RequestsToUpdate.Num() > 0)
+	{
+		ENQUEUE_RENDER_COMMAND(FRenderTraceUpdaterTick)(
+			[InRequestsToUpdate=MoveTemp(RequestsToUpdate)](FRHICommandListImmediate& RHICmdList)
+			{
+				UpdateTasks_RenderThread(RHICmdList, InRequestsToUpdate, false);
+			});
+	}
+
+	auto CheckCompletionAndRemove = [this](const TSharedPtr<FRenderTraceTask>& Task)
+	{
+		if (Task->CompletionState.load() == ECompletionState::Complete)
+		{
+#if RENDER_TRACE_LOG_TIMING
+			double TaskDurationSeconds = FPlatformTime::Seconds() - Task->StartSeconds;
+			int32 TaskDurationFrames = int32(GFrameCounter - Task->StartFrame);
+			UE_LOG(LogTemp, Log, TEXT("RenderTrace completed in %f seconds (%d frames)"), TaskDurationSeconds, TaskDurationFrames);
+#endif
+
+			INC_DWORD_STAT(STAT_RenderTrace_Completed);
+			if (Task->ResultDelegate.IsBound())
+			{
+				Task->ResultDelegate.Execute(Task->RequestID, Task->ResultMaterial, Task->UserData);
+			}
+			else
+			{
+				UE_LOG(LogTemp, Warning, TEXT("RenderTrace completed with an unbound delegate"));
+			}
+			return true;
+		}
+
+		return false;
+	};
+
+	RequestsInFlight.RemoveAll(CheckCompletionAndRemove);
+}
+
+bool FRenderTraceQueue::IsEnabled()
+{
+	return bEnableRenderTracing;
+}
+
diff --git a/Engine/Plugins/Runtime/RenderTrace/Source/Private/RenderTraceModule.cpp b/Engine/Plugins/Runtime/RenderTrace/Source/Private/RenderTraceModule.cpp
new file mode 100644
--- /dev/null
+++ b/Engine/Plugins/Runtime/RenderTrace/Source/Private/RenderTraceModule.cpp
@@ -0,0 +1,26 @@
+// Copyright Epic Games, Inc. All Rights Reserved.
+
+#include "RenderTraceModule.h"
+#include "Modules/ModuleManager.h"
+#include "Interfaces/IPluginManager.h"
+
+class FRenderTraceModule : public IModuleInterface
+{
+public:
+	/** IModuleInterface implementation */
+	void StartupModule() override;
+	void ShutdownModule() override;
+};
+
+
+void FRenderTraceModule::StartupModule()
+{
+	FString PluginShaderDir = FPaths::Combine(IPluginManager::Get().FindPlugin(TEXT("RenderTrace"))->GetBaseDir(), TEXT("Shaders"));
+	AddShaderSourceDirectoryMapping(TEXT("/Plugin/Runtime/RenderTrace"), PluginShaderDir);
+}
+
+void FRenderTraceModule::ShutdownModule()
+{
+}
+
+IMPLEMENT_MODULE(FRenderTraceModule, RenderTrace);
diff --git a/Engine/Plugins/Runtime/RenderTrace/Source/Public/RenderTrace.h b/Engine/Plugins/Runtime/RenderTrace/Source/Public/RenderTrace.h
new file mode 100644
--- /dev/null
+++ b/Engine/Plugins/Runtime/RenderTrace/Source/Public/RenderTrace.h
@@ -0,0 +1,45 @@
+// Copyright Epic Games, Inc. All Rights Reserved.
+
+#pragma once
+
+#include "CoreMinimal.h"
+#include "Tickable.h"
+#include "Logging/LogMacros.h"
+
+DECLARE_LOG_CATEGORY_EXTERN(LogRenderTrace, Log, All);
+
+DECLARE_DELEGATE_ThreeParams(FRenderTraceDelegate, uint32 TaskID, class UPhysicalMaterial const*, int64 UserData);
+
+/** 
+ * The work uses the GPU so this object handles reading back the data without stalling.
+ * All calls are expected to be made on the game thread only.
+ */
+class RENDERTRACE_API FRenderTraceQueue : FTickableGameObject
+{
+public:
+	FRenderTraceQueue() = default;
+	~FRenderTraceQueue() = default;
+	UE_NONCOPYABLE(FRenderTraceQueue);
+
+	/** 
+	 *	Initialize the task for a list of components. Previously queued tasks will continue to process. 
+	 *	Returns the unique ID of the task that will be sent with the OnCompletion callback, or 0 if the 
+	 *  request was invalid or there were no valid primitives to check. 
+	 */
+	uint32 AsyncRenderTraceComponents(TArrayView<const class UPrimitiveComponent*> PrimitiveComponents, FVector RayOrigin, FVector RayDirection, FRenderTraceDelegate OnComplete, int64 UserData = 0);
+	void CancelAsyncSample(uint32 RequestID);
+
+	/* FTickableGameObject begin */ 
+	virtual void Tick(float DeltaTime) override;
+	virtual TStatId GetStatId() const override { RETURN_QUICK_DECLARE_CYCLE_STAT(FRenderTraceStat, STATGROUP_Tickables); }
+	virtual bool IsTickable() const override { return !RequestsInFlight.IsEmpty(); }
+	virtual bool IsTickableInEditor() const override { return false; }
+	/* FTickableGameObject end */
+
+	static bool IsEnabled();
+
+private:
+	uint32 LastRequestID = 0;
+	TArray<TSharedPtr<struct FRenderTraceTask>> RequestsInFlight;
+};
+
diff --git a/Engine/Plugins/Runtime/RenderTrace/Source/Public/RenderTraceModule.h b/Engine/Plugins/Runtime/RenderTrace/Source/Public/RenderTraceModule.h
new file mode 100644
--- /dev/null
+++ b/Engine/Plugins/Runtime/RenderTrace/Source/Public/RenderTraceModule.h
@@ -0,0 +1,14 @@
+// Copyright Epic Games, Inc. All Rights Reserved.
+
+#pragma once
+
+#include "CoreMinimal.h"
+#include "Modules/ModuleInterface.h"
+
+/**
+* RenderTrace module interface
+*/
+class IRenderTraceModule : public IModuleInterface
+{
+public:
+};
diff --git a/Engine/Plugins/Runtime/RenderTrace/Source/RenderTrace.Build.cs b/Engine/Plugins/Runtime/RenderTrace/Source/RenderTrace.Build.cs
new file mode 100644
--- /dev/null
+++ b/Engine/Plugins/Runtime/RenderTrace/Source/RenderTrace.Build.cs
@@ -0,0 +1,52 @@
+// Copyright Epic Games, Inc. All Rights Reserved.
+
+using UnrealBuildTool;
+using System.IO;
+
+public class RenderTrace : ModuleRules
+{
+	public RenderTrace(ReadOnlyTargetRules Target) : base(Target)
+	{
+		PrivateIncludePaths.AddRange(
+			new string[] {
+				"Private",
+                //"../Shaders/Shared"
+            }
+		);
+
+		PrivateIncludePathModuleNames.AddRange(
+			new string[] {
+				"Renderer",
+				"TargetPlatform",
+			}
+		);
+
+		PrivateDependencyModuleNames.AddRange(
+			new string[] {
+				"ApplicationCore",
+				"Core",
+				"CoreUObject",
+				"DeveloperSettings",
+				"Engine",
+				"Projects",
+				"RHI",
+				"RenderCore", 
+				"Renderer",
+			}
+		);
+		SetupModulePhysicsSupport(Target);
+
+		if (Target.bBuildEditor == true)
+		{
+			PrivateDependencyModuleNames.AddRange(
+				new string[] {
+					"EditorFramework",
+					"MaterialUtilities", 
+					"Slate",
+					"SlateCore",
+					"UnrealEd",
+				}
+			);
+		}
+	}
+}
diff --git a/Engine/Source/Runtime/Engine/Classes/Materials/Material.h b/Engine/Source/Runtime/Engine/Classes/Materials/Material.h
--- a/Engine/Source/Runtime/Engine/Classes/Materials/Material.h
+++ b/Engine/Source/Runtime/Engine/Classes/Materials/Material.h
@@ -416,6 +416,10 @@ class UMaterial : public UMaterialInterface
 	/** Physical material mask map to use for this graphics material. Used for sounds, effects etc.*/
 	UPROPERTY(EditAnywhere, Category = PhysicalMaterialMask)
 	TObjectPtr<class UPhysicalMaterial> PhysicalMaterialMap[EPhysicalMaterialMaskColor::MAX];
+	
+
+	UPROPERTY()
+	TArray<TObjectPtr<class UPhysicalMaterial>> RenderTracePhysicalMaterialOutputs;
 
 	/** 
 	 * The domain that the material's attributes will be evaluated in. 
@@ -1770,6 +1774,16 @@ public:
 	static bool FixFeatureLevelNodesForSM6(TArray<UMaterialExpression*> const& InExpressions);
 
 #endif //WITH_EDITORONLY_DATA
+	/**
+	 * Return the array of physical materials associated with the RenderTrace material outputs.
+	 * @return The physical material array.
+	 */
+	ENGINE_API TArrayView<const TObjectPtr<UPhysicalMaterial>> GetRenderTracePhysicalMaterialOutputs() const;
+
+	/**
+	 * Set the array of physical materials associated with the RenderTrace material outputs.
+	 */
+	ENGINE_API void SetRenderTracePhysicalMaterialOutputs(TArrayView<TObjectPtr<UPhysicalMaterial>> PhysicalMaterials);
 
 private:
 #if WITH_EDITOR
diff --git a/Engine/Source/Runtime/Engine/Private/Materials/Material.cpp b/Engine/Source/Runtime/Engine/Private/Materials/Material.cpp
--- a/Engine/Source/Runtime/Engine/Private/Materials/Material.cpp
+++ b/Engine/Source/Runtime/Engine/Private/Materials/Material.cpp
@@ -2048,6 +2048,16 @@ UPhysicalMaterial* UMaterial::GetPhysicalMaterialFromMap(int32 Index) const
 	return nullptr;
 }
 
+TArrayView<const TObjectPtr<UPhysicalMaterial>> UMaterial::GetRenderTracePhysicalMaterialOutputs() const
+{
+	return MakeArrayView(RenderTracePhysicalMaterialOutputs);
+}
+
+void UMaterial::SetRenderTracePhysicalMaterialOutputs(TArrayView<TObjectPtr<class UPhysicalMaterial>> PhysicalMaterials)
+{
+	RenderTracePhysicalMaterialOutputs = PhysicalMaterials;
+}
+
 /** Helper functions for text output of properties... */
 #ifndef CASE_ENUM_TO_TEXT
 #define CASE_ENUM_TO_TEXT(txt) case txt: return TEXT(#txt);
diff --git a/Engine/Source/Runtime/Engine/Private/Materials/MaterialShared.cpp b/Engine/Source/Runtime/Engine/Private/Materials/MaterialShared.cpp
--- a/Engine/Source/Runtime/Engine/Private/Materials/MaterialShared.cpp
+++ b/Engine/Source/Runtime/Engine/Private/Materials/MaterialShared.cpp
@@ -1898,6 +1898,11 @@ bool FMaterialResource::CastsRayTracedShadows() const
 	return Material->bCastRayTracedShadows;
 }
 
+bool FMaterialResource::HasRenderTracePhysicalMaterialOutputs() const
+{
+	return Material->GetRenderTracePhysicalMaterialOutputs().Num() > 0;
+}
+
 UMaterialInterface* FMaterialResource::GetMaterialInterface() const 
 { 
 	return MaterialInstance ? (UMaterialInterface*)MaterialInstance : (UMaterialInterface*)Material;
@@ -5632,6 +5637,7 @@ FMaterialShaderParameters::FMaterialShaderParameters(const FMaterial* InMaterial
 	bIsStencilTestEnabled = InMaterial->IsStencilTestEnabled();
 	bIsTranslucencySurface = InMaterial->GetTranslucencyLightingMode() == ETranslucencyLightingMode::TLM_Surface || InMaterial->GetTranslucencyLightingMode() == ETranslucencyLightingMode::TLM_SurfacePerPixelLighting;
 	bShouldDisableDepthTest = InMaterial->ShouldDisableDepthTest();
+	bHasRenderTracePhysicalMaterialOutput = InMaterial->HasRenderTracePhysicalMaterialOutputs();
 }
 
 #undef LOCTEXT_NAMESPACE
diff --git a/Engine/Source/Runtime/Engine/Public/MaterialShared.h b/Engine/Source/Runtime/Engine/Public/MaterialShared.h
--- a/Engine/Source/Runtime/Engine/Public/MaterialShared.h
+++ b/Engine/Source/Runtime/Engine/Public/MaterialShared.h
@@ -1915,6 +1915,7 @@ public:
 	virtual bool HasVertexInterpolator() const { return false; }
 	virtual bool HasRuntimeVirtualTextureOutput() const { return false; }
 	virtual bool CastsRayTracedShadows() const { return true; }
+	virtual bool HasRenderTracePhysicalMaterialOutputs() const { return false; }
 	virtual EMaterialShadingRate GetShadingRate() const { return MSR_1x1; }
 	/**
 	 * Should shaders compiled for this material be saved to disk?
@@ -2885,7 +2886,8 @@ public:
 	ENGINE_API virtual bool HasVertexInterpolator() const override;
 	ENGINE_API virtual bool HasRuntimeVirtualTextureOutput() const override;
 	ENGINE_API virtual bool CastsRayTracedShadows() const override;
-	ENGINE_API  virtual UMaterialInterface* GetMaterialInterface() const override;
+	ENGINE_API virtual bool HasRenderTracePhysicalMaterialOutputs() const override;
+	ENGINE_API virtual UMaterialInterface* GetMaterialInterface() const override;
 	/**
 	 * Should shaders compiled for this material be saved to disk?
 	 */
@@ -3471,6 +3473,7 @@ struct FMaterialShaderParameters
 			uint64 bIsStencilTestEnabled : 1;
 			uint64 bIsTranslucencySurface : 1;
 			uint64 bShouldDisableDepthTest : 1;
+			uint64 bHasRenderTracePhysicalMaterialOutput : 1;
 		};
 	};
 
