Add includes to fix non unity builds

diff --git a/Engine/Plugins/Runtime/RenderTrace/Source/Private/Materials/MaterialExpressionPhysicalMaterialOutput.cpp b/Engine/Plugins/Runtime/RenderTrace/Source/Private/Materials/MaterialExpressionPhysicalMaterialOutput.cpp
--- a/Engine/Plugins/Runtime/RenderTrace/Source/Private/Materials/MaterialExpressionPhysicalMaterialOutput.cpp
+++ b/Engine/Plugins/Runtime/RenderTrace/Source/Private/Materials/MaterialExpressionPhysicalMaterialOutput.cpp
@@ -5,6 +5,7 @@
 #include "Engine/Engine.h"
 #include "EngineGlobals.h"
 #include "MaterialCompiler.h"
+#include "Materials/Material.h"
 #include "PhysicalMaterials/PhysicalMaterial.h"
 
 #define LOCTEXT_NAMESPACE "RenderTrace"
diff --git a/Engine/Plugins/Runtime/RenderTrace/Source/Private/RenderTraceModule.cpp b/Engine/Plugins/Runtime/RenderTrace/Source/Private/RenderTraceModule.cpp
--- a/Engine/Plugins/Runtime/RenderTrace/Source/Private/RenderTraceModule.cpp
+++ b/Engine/Plugins/Runtime/RenderTrace/Source/Private/RenderTraceModule.cpp
@@ -3,6 +3,8 @@
 #include "RenderTraceModule.h"
 #include "Modules/ModuleManager.h"
 #include "Interfaces/IPluginManager.h"
+#include "Misc/Paths.h"
+#include "ShaderCore.h"
 
 class FRenderTraceModule : public IModuleInterface
 {
