Fixing compile errors

Removing calls to deprecated functions
Adding include

diff --git a/Engine/Plugins/Runtime/RenderTrace/Source/Private/RenderTrace.cpp b/Engine/Plugins/Runtime/RenderTrace/Source/Private/RenderTrace.cpp
--- a/Engine/Plugins/Runtime/RenderTrace/Source/Private/RenderTrace.cpp
+++ b/Engine/Plugins/Runtime/RenderTrace/Source/Private/RenderTrace.cpp
@@ -12,6 +12,8 @@
 #include "MeshPassProcessor.inl"
 #include "RenderCommandFence.h"
 #include "RenderTargetPool.h"
+#include "RenderGraphUtils.h"
+#include "RHIGPUReadback.h"
 #include "RHIResources.h"
 #include "SceneRenderTargetParameters.h"
 #include "SimpleMeshDrawCommandPass.h"
@@ -374,7 +376,6 @@ namespace
 		
 		// Create the view
 		FSceneViewFamily::ConstructionValues ViewFamilyInit(nullptr, SceneInterface, FEngineShowFlags(ESFIM_Game));
-		ViewFamilyInit.SetWorldTimes(0.0f, 0.0f, 0.0f);
 		FSceneViewFamilyContext ViewFamily(ViewFamilyInit);
 
 		FScenePrimitiveRenderingContextScopeHelper ScenePrimitiveRenderingContextScopeHelper(GetRendererModule().BeginScenePrimitiveRendering(GraphBuilder, &ViewFamily));
@@ -442,11 +443,13 @@ namespace
 		FRHIGPUTextureReadback* Readback,
 		FRenderTraceReadbackData& OutSampleResult)
 	{
-		void* Data = nullptr;
 		int32 Pitch;
-		Readback->LockTexture(RHICmdList, Data, Pitch);
+		void* Data = Readback->Lock(Pitch);
 		check(Data && TargetSize.X <= Pitch);
-		FMemory::Memcpy(&OutSampleResult, Data, sizeof(OutSampleResult));
+		if (Data)
+		{
+			FMemory::Memcpy(&OutSampleResult, Data, sizeof(OutSampleResult));
+		}
 		Readback->Unlock();
 	}
 
