Fix assert in SceneView when initializing a scene from RenderTrace.

diff --git a/Engine/Plugins/Runtime/RenderTrace/Source/Private/RenderTrace.cpp b/Engine/Plugins/Runtime/RenderTrace/Source/Private/RenderTrace.cpp
--- a/Engine/Plugins/Runtime/RenderTrace/Source/Private/RenderTrace.cpp
+++ b/Engine/Plugins/Runtime/RenderTrace/Source/Private/RenderTrace.cpp
@@ -376,6 +376,8 @@ namespace
 		
 		// Create the view
 		FSceneViewFamily::ConstructionValues ViewFamilyInit(nullptr, SceneInterface, FEngineShowFlags(ESFIM_Game));
+		FGameTime Time;
+		ViewFamilyInit.SetTime(Time);
 		FSceneViewFamilyContext ViewFamily(ViewFamilyInit);
 
 		FScenePrimitiveRenderingContextScopeHelper ScenePrimitiveRenderingContextScopeHelper(GetRendererModule().BeginScenePrimitiveRendering(GraphBuilder, &ViewFamily));
