Remove unnecessary logging

diff --git a/Engine/Plugins/Runtime/RenderTrace/Source/Private/Materials/MaterialExpressionPhysicalMaterialOutput.cpp b/Engine/Plugins/Runtime/RenderTrace/Source/Private/Materials/MaterialExpressionPhysicalMaterialOutput.cpp
--- a/Engine/Plugins/Runtime/RenderTrace/Source/Private/Materials/MaterialExpressionPhysicalMaterialOutput.cpp
+++ b/Engine/Plugins/Runtime/RenderTrace/Source/Private/Materials/MaterialExpressionPhysicalMaterialOutput.cpp
@@ -143,8 +143,6 @@ void UMaterialExpressionPhysicalMaterialOutput::PostEditChangeProperty(FProperty
 
 int32 UMaterialExpressionPhysicalMaterialOutput::GetNumOutputs() const 
 { 
-	int32 Count = (int32)Inputs.Num(); 
-	UE_LOG(LogTemp, Warning, TEXT("UMaterialExpressionPhysicalMaterialOutput::GetNumOutputs() is returning %d"), Count);
 	return Inputs.Num(); 
 }
 
