Added missing scene uniform buffer at slot 1 in RenderTrace

diff --git a/Engine/Plugins/Runtime/RenderTrace/Source/Private/RenderTrace.cpp b/Engine/Plugins/Runtime/RenderTrace/Source/Private/RenderTrace.cpp
--- a/Engine/Plugins/Runtime/RenderTrace/Source/Private/RenderTrace.cpp
+++ b/Engine/Plugins/Runtime/RenderTrace/Source/Private/RenderTrace.cpp
@@ -14,6 +14,7 @@
 #include "DataDrivenShaderPlatformInfo.h"
 #include "SceneTexturesConfig.h"
 #include "SimpleMeshDrawCommandPass.h"
+#include "SceneRendererInterface.h"
 
 DEFINE_LOG_CATEGORY(LogRenderTrace);
 
@@ -351,6 +352,7 @@ namespace
 	
 	BEGIN_SHADER_PARAMETER_STRUCT(FPhysicalMaterialSamplerPassParameters, )
 		SHADER_PARAMETER_STRUCT_REF(FViewUniformShaderParameters, View)
+		SHADER_PARAMETER_RDG_UNIFORM_BUFFER(FSceneUniformParameters, Scene)
 		SHADER_PARAMETER_STRUCT_INCLUDE(FInstanceCullingDrawParams, InstanceCullingDrawParams)
 		RENDER_TARGET_BINDING_SLOTS()
 	END_SHADER_PARAMETER_STRUCT()
@@ -402,6 +404,7 @@ namespace
 
 			auto* PassParameters = GraphBuilder.AllocParameters<FPhysicalMaterialSamplerPassParameters>();
 			PassParameters->View = View->ViewUniformBuffer;
+			PassParameters->Scene = GetSceneUniformBufferRef(GraphBuilder, *View);
 			PassParameters->RenderTargets[0] = FRenderTargetBinding(OutputTexture, ERenderTargetLoadAction::EClear);
 			PassParameters->RenderTargets.DepthStencil = FDepthStencilBinding(DepthTexture, ERenderTargetLoadAction::EClear, ERenderTargetLoadAction::ENoAction, FExclusiveDepthStencil::DepthWrite_StencilNop);
 
